Duplicate PR #2855:  Ansel plugin fixes (Contributed by [REDACTED_PERSON])

diff --git a/Engine/Plugins/Runtime/Nvidia/Ansel/Ansel.uplugin b/Engine/Plugins/Runtime/Nvidia/Ansel/Ansel.uplugin
--- a/Engine/Plugins/Runtime/Nvidia/Ansel/Ansel.uplugin
+++ b/Engine/Plugins/Runtime/Nvidia/Ansel/Ansel.uplugin
@@ -2,8 +2,8 @@
 	"FileVersion" : 3,
 	"Version" : 1,
 	"VersionName" : "1.0",
-	"FriendlyName" : "Nvidia Ansel Plugin",
-	"Description" : "Plugin to support Ansel",
+	"FriendlyName" : "Nvidia Ansel Photography Plugin",
+	"Description" : "Plugin to support Ansel Photography",
 	"Category" : "Photography",
 	"CreatedBy" : "Epic Games, Inc.",
 	"CreatedByURL" : "http://epicgames.com",
diff --git a/Engine/Plugins/Runtime/Nvidia/Ansel/Source/Ansel/Ansel.Build.cs b/Engine/Plugins/Runtime/Nvidia/Ansel/Source/Ansel/Ansel.Build.cs
--- a/Engine/Plugins/Runtime/Nvidia/Ansel/Source/Ansel/Ansel.Build.cs
+++ b/Engine/Plugins/Runtime/Nvidia/Ansel/Source/Ansel/Ansel.Build.cs
@@ -69,7 +69,7 @@ public class Ansel : ModuleRules
 
             string DLLName = LibName + ".dll";
             PublicDelayLoadDLLs.Add(DLLName);
-            RuntimeDependencies.Add(new RuntimeDependency("BinariesPath" + DLLName));
+            RuntimeDependencies.Add(new RuntimeDependency(BinariesPath + DLLName));
 
             Definitions.Add("WITH_ANSEL=1");
             Definitions.Add("ANSEL_DLL=" + DLLName);
diff --git a/Engine/Plugins/Runtime/Nvidia/Ansel/Source/Ansel/Private/Ansel.cpp b/Engine/Plugins/Runtime/Nvidia/Ansel/Source/Ansel/Private/Ansel.cpp
--- a/Engine/Plugins/Runtime/Nvidia/Ansel/Source/Ansel/Private/Ansel.cpp
+++ b/Engine/Plugins/Runtime/Nvidia/Ansel/Source/Ansel/Private/Ansel.cpp
@@ -69,6 +69,7 @@ private:
 };
 
 static HMODULE AnselSDKDLLHandle = 0;
+static bool bAnselDLLLoaded = false;
 
 FNVAnselCameraPhotographyPrivate::FNVAnselCameraPhotographyPrivate()
 	: ICameraPhotography()
@@ -81,8 +82,7 @@ FNVAnselCameraPhotographyPrivate::FNVAnselCameraPhotographyPrivate()
 	, bForceDisallow(false)
 	, bIsOrthoProjection(false)
 {
-	bool bLoadedDLL = AnselSDKDLLHandle != 0;
-	if (bLoadedDLL)
+	if (bAnselDLLLoaded)
 	{
 		AnselConfig = new ansel::Configuration();
 
@@ -115,22 +115,25 @@ FNVAnselCameraPhotographyPrivate::FNVAnselCameraPhotographyPrivate()
 	}
 	else
 	{
-		UE_LOG(LogAnsel, Log, TEXT("Ansel DLL was not loaded."));
+		UE_LOG(LogAnsel, Log, TEXT("Ansel DLL was not successfully loaded."));
 	}
 }
 
 
 FNVAnselCameraPhotographyPrivate::~FNVAnselCameraPhotographyPrivate()
 {	
-	IConsoleManager::Get().UnregisterConsoleVariableSink_Handle(CVarDelegateHandle);
-	DeconfigureAnsel();
-	delete AnselConfig;
+	if (bAnselDLLLoaded)
+	{
+		IConsoleManager::Get().UnregisterConsoleVariableSink_Handle(CVarDelegateHandle);
+		DeconfigureAnsel();
+		delete AnselConfig;
+	}
 }
 
 
 bool FNVAnselCameraPhotographyPrivate::IsSupported()
 {
-	return ansel::isAnselAvailable();
+	return bAnselDLLLoaded && ansel::isAnselAvailable();
 }
 
 bool FNVAnselCameraPhotographyPrivate::AnselCamerasMatch(ansel::Camera& a, ansel::Camera& b)
@@ -348,26 +351,32 @@ bool FNVAnselCameraPhotographyPrivate::UpdateCamera(FMinimalViewInfo& InOutPOV,
 					// force-disable the standard postprocessing effects which are known to
 					// be problematic in multi-part shots
 
+					// these effects tile poorly
 					InOutPOV.PostProcessSettings.bOverride_BloomDirtMaskIntensity = 1;
 					InOutPOV.PostProcessSettings.BloomDirtMaskIntensity = 0.f;
-
 					InOutPOV.PostProcessSettings.bOverride_LensFlareIntensity = 1;
 					InOutPOV.PostProcessSettings.LensFlareIntensity = 0.f;
-
 					InOutPOV.PostProcessSettings.bOverride_VignetteIntensity = 1;
 					InOutPOV.PostProcessSettings.VignetteIntensity = 0.f;
-
-					InOutPOV.PostProcessSettings.bOverride_DepthOfFieldScale = 1;
-					InOutPOV.PostProcessSettings.DepthOfFieldScale = 0.f;
-					InOutPOV.PostProcessSettings.bOverride_DepthOfFieldMaxBokehSize = 1;
-					InOutPOV.PostProcessSettings.DepthOfFieldMaxBokehSize = 0.f;
-
 					InOutPOV.PostProcessSettings.bOverride_SceneFringeIntensity = 1;
 					InOutPOV.PostProcessSettings.SceneFringeIntensity = 0.f;
 
+					// motion blur doesn't make sense with a teleporting camera
 					InOutPOV.PostProcessSettings.bOverride_MotionBlurAmount = 1;
 					InOutPOV.PostProcessSettings.MotionBlurAmount = 0.f;
 
+					// DoF can look poor/wrong at high-res, depending on settings.
+					InOutPOV.PostProcessSettings.bOverride_DepthOfFieldScale = 1;
+					InOutPOV.PostProcessSettings.DepthOfFieldScale = 0.f; // BokehDOF
+					InOutPOV.PostProcessSettings.bOverride_DepthOfFieldNearBlurSize = 1;
+					InOutPOV.PostProcessSettings.DepthOfFieldNearBlurSize = 0.f; // GaussianDOF
+					InOutPOV.PostProcessSettings.bOverride_DepthOfFieldFarBlurSize = 1;
+					InOutPOV.PostProcessSettings.DepthOfFieldFarBlurSize = 0.f; // GaussianDOF
+					InOutPOV.PostProcessSettings.bOverride_DepthOfFieldDepthBlurRadius = 1;
+					InOutPOV.PostProcessSettings.DepthOfFieldDepthBlurRadius = 0.f; // CircleDOF
+					InOutPOV.PostProcessSettings.bOverride_DepthOfFieldVignetteSize = 1;
+					InOutPOV.PostProcessSettings.DepthOfFieldVignetteSize = 200.f; // Scene.h says 200.0 means 'no effect'
+
 					// freeze auto-exposure adaptation
 					InOutPOV.PostProcessSettings.bOverride_AutoExposureSpeedDown = 1;
 					InOutPOV.PostProcessSettings.AutoExposureSpeedDown = 0.f;
@@ -513,30 +522,29 @@ public:
 	virtual void StartupModule() override
 	{
 		ICameraPhotographyModule::StartupModule();
+		check(!bAnselDLLLoaded);
 
 		// Late-load Ansel DLL.  DLL name has been worked out by the build scripts as ANSEL_DLL
-		bool bLoadedDLL = false;
 		FString AnselDLLName;
-		if (0 == AnselSDKDLLHandle)
-		{
-			FString AnselBinariesRoot = FPaths::EngineDir() / TEXT("Plugins/Runtime/Nvidia/Ansel/Binaries/ThirdParty/");
-			// common preprocessor fudge to convert macro expansion into string
+		FString AnselBinariesRoot = FPaths::EngineDir() / TEXT("Plugins/Runtime/Nvidia/Ansel/Binaries/ThirdParty/");
+		// common preprocessor fudge to convert macro expansion into string
 #define STRINGIFY(X) STRINGIFY2(X)
 #define STRINGIFY2(X) #X
-			AnselDLLName = AnselBinariesRoot + TEXT(STRINGIFY(ANSEL_DLL));
-			AnselSDKDLLHandle = LoadLibraryW(*(AnselDLLName));
+		AnselDLLName = AnselBinariesRoot + TEXT(STRINGIFY(ANSEL_DLL));
+		AnselSDKDLLHandle = LoadLibraryW(*(AnselDLLName));
 
-			bLoadedDLL = AnselSDKDLLHandle != 0;			
-		}
+		bAnselDLLLoaded = AnselSDKDLLHandle != 0;
+		
+		UE_LOG(LogAnsel, Log, TEXT("Tried to load %s : success=%d"), *AnselDLLName, int(bAnselDLLLoaded));
 	}
 
 	virtual void ShutdownModule() override
 	{		
-		bool bLoadedDLL = AnselSDKDLLHandle != 0;
-		if (bLoadedDLL)
+		if (bAnselDLLLoaded)
 		{
 			FreeLibrary(AnselSDKDLLHandle);
 			AnselSDKDLLHandle = 0;
+			bAnselDLLLoaded = false;
 		}
 		ICameraPhotographyModule::ShutdownModule();
 	}
@@ -545,11 +553,18 @@ private:
 	virtual TSharedPtr< class ICameraPhotography > CreateCameraPhotography() override
 	{
 		TSharedPtr<ICameraPhotography> Photography = nullptr;
-		const bool AnselAvailable = ansel::isAnselAvailable();
-		if (AnselAvailable)
+
+		FNVAnselCameraPhotographyPrivate* PhotographyPrivate = nullptr;
+		PhotographyPrivate = new FNVAnselCameraPhotographyPrivate();
+		if (PhotographyPrivate->IsSupported())
 		{
-			Photography = TSharedPtr<ICameraPhotography>(new FNVAnselCameraPhotographyPrivate());
+			Photography = TSharedPtr<ICameraPhotography>(PhotographyPrivate);
 		}
+		else
+		{
+			delete PhotographyPrivate;
+		}
+
 		return Photography;
 	}
 };
diff --git a/Engine/Source/Runtime/Engine/Public/CameraPhotographyModule.h b/Engine/Source/Runtime/Engine/Public/CameraPhotographyModule.h
--- a/Engine/Source/Runtime/Engine/Public/CameraPhotographyModule.h
+++ b/Engine/Source/Runtime/Engine/Public/CameraPhotographyModule.h
@@ -26,7 +26,7 @@ public:
 	*/
 	static inline ICameraPhotographyModule& Get()
 	{
-		return FModuleManager::LoadModuleChecked< ICameraPhotographyModule >("InputDevice");
+		return FModuleManager::LoadModuleChecked< ICameraPhotographyModule >(GetModularFeatureName());
 	}
 
 	/**
@@ -36,13 +36,13 @@ public:
 	*/
 	static inline bool IsAvailable()
 	{
-		return FModuleManager::Get().IsModuleLoaded("InputDevice");
+		return FModuleManager::Get().IsModuleLoaded(GetModularFeatureName());
 	}
 
 	/**
-	* Attempts to create a new head tracking device interface
+	* Attempts to create a new photography interface
 	*
-	* @return	Interface to the new head tracking device, if we were able to successfully create one
+	* @return Interface to the photography implementation, if we were able to successfully create one
 	*/
 	virtual TSharedPtr< class ICameraPhotography > CreateCameraPhotography() = 0;
 };
