UE4.22 - Fix compile issue
#rnx

diff --git a/Engine/Plugins/Runtime/Nvidia/Ansel/Source/Ansel/Private/Ansel.cpp b/Engine/Plugins/Runtime/Nvidia/Ansel/Source/Ansel/Private/Ansel.cpp
--- a/Engine/Plugins/Runtime/Nvidia/Ansel/Source/Ansel/Private/Ansel.cpp
+++ b/Engine/Plugins/Runtime/Nvidia/Ansel/Source/Ansel/Private/Ansel.cpp
@@ -124,15 +124,15 @@ private:
 	bool bWasShowingHUDBeforeSession;
 	bool bWereSubtitlesEnabledBeforeSession;
 	bool bWasFadingEnabledBeforeSession;
-	bool bWasScreenMessagesEnabledBeforeSession;
+	bool bWasScreenMessagesEnabledBeforeSession = false;
 
-	bool bCameraIsInOriginalState;
+	bool bCameraIsInOriginalState = true;
 
 	bool bAutoPostprocess;
 	bool bAutoPause;
 
-	bool bHighQualityModeDesired;
-	bool bHighQualityModeIsSetup;
+	bool bHighQualityModeDesired = false;
+	bool bHighQualityModeIsSetup = false;
 
 	// members relating to the 'Game Settings' controls in the Ansel overlay UI
 	TStaticBitArray<256> bEffectUIAllowed;
@@ -183,8 +183,6 @@ FNVAnselCameraPhotographyPrivate::FNVAnselCameraPhotographyPrivate()
 	, bForceDisallow(false)
 	, bIsOrthoProjection(false)
 	, bUIControlsNeedRebuild(false)
-	, bHighQualityModeIsSetup(false)
-	, bCameraIsInOriginalState(true)
 {
 	for (int i = 0; i < bEffectUIAllowed.Num(); ++i)
 	{
