PR #5610: Ansel photography: raytracing boosts for 'high quality' mode (4.22-ansel53) (Contributed by [REDACTED_PERSON])

diff --git a/Engine/Documentation/Source/Engine/Plugins/Ansel/Reference/ConsoleVariables/AnselConsoleVars.INT.udn b/Engine/Documentation/Source/Engine/Plugins/Ansel/Reference/ConsoleVariables/AnselConsoleVars.INT.udn
--- a/Engine/Documentation/Source/Engine/Plugins/Ansel/Reference/ConsoleVariables/AnselConsoleVars.INT.udn
+++ b/Engine/Documentation/Source/Engine/Plugins/Ansel/Reference/ConsoleVariables/AnselConsoleVars.INT.udn
@@ -24,8 +24,8 @@ The following console variables are available for configuring the behavior of th
 | `r.Photography.EnableMultipart` | 1 | If 1, allow the photography system to take high-resolution shots that need to be rendered in tiles which are later stitched together. |
 | `r.Photography.SettleFrames` | 10 | The number of frames to let the rendering 'settle' before taking a photo.  Useful to allow temporal AA/smoothing to work well; if not using any temporal effects, can be lowered for faster capture. |
 | `r.Photography.TranslationSpeed` | 100.0 | Normal speed (in Unreal units per second) at which to move the roaming photography camera. |
-| `r.Photography.Constrain.CameraSize` | 14.0 | Radius (in Unreal Units) of sphere around the camera; used to prevent the camera clipping into nearby geometry when constraining camera with collision. |
-| `r.Photography.Constrain.CameraDistance` | 2500.0 | Maximum distance (in Unreal Units) which camera is allowed to wander from its initial position when constraining camera by distance. |
+| `r.Photography.Constrain.CameraSize` | 14.0 | Radius (in Unreal Units) of sphere around the camera; used to prevent the camera clipping into nearby geometry when constraining camera with collision.  <0 disables camera collision. |
+| `r.Photography.Constrain.CameraDistance` | 2500.0 | Maximum distance (in Unreal Units) which camera is allowed to wander from its initial position when constraining camera by distance.  <0 disables distance constraint. |
 | `r.Photography.AutoPostprocess` | 1 | If 1, the photography system will attempt to automatically disable HUD, subtitles, and some standard postprocessing effects during photography sessions/captures which are known to give poor photography results.  Set to 0 to manage all postprocessing tweaks manually from the PlayerCameraManager Blueprint callbacks. [REGION:note]Blueprint callbacks will be called regardless of AutoPostprocess value.[/REGION] |
 | `r.Photography.AutoPause` | 1 | If 1, the photography system will attempt to ensure that the level is paused while in photography mode.  Set to 0 to manage pausing and unpausing manually from the PlayerCameraManager Blueprint callbacks.  [REGION:note]Blueprint callbacks will be called regardless of AutoPause value.[/REGION] |
 
diff --git a/Engine/Plugins/Runtime/Nvidia/Ansel/Source/Ansel/Ansel.Build.cs b/Engine/Plugins/Runtime/Nvidia/Ansel/Source/Ansel/Ansel.Build.cs
--- a/Engine/Plugins/Runtime/Nvidia/Ansel/Source/Ansel/Ansel.Build.cs
+++ b/Engine/Plugins/Runtime/Nvidia/Ansel/Source/Ansel/Ansel.Build.cs
@@ -86,6 +86,7 @@ public class Ansel : ModuleRules
 			"CoreUObject",
             "SlateCore",
 			"Engine",
+			"RHI",
 		});
 
         PublicDependencyModuleNames.Add("Engine");
diff --git a/Engine/Plugins/Runtime/Nvidia/Ansel/Source/Ansel/Private/Ansel.cpp b/Engine/Plugins/Runtime/Nvidia/Ansel/Source/Ansel/Private/Ansel.cpp
--- a/Engine/Plugins/Runtime/Nvidia/Ansel/Source/Ansel/Private/Ansel.cpp
+++ b/Engine/Plugins/Runtime/Nvidia/Ansel/Source/Ansel/Private/Ansel.cpp
@@ -17,6 +17,9 @@
 #include "Kismet/GameplayStatics.h"
 #include "Widgets/SWindow.h"
 #include "Application/SlateApplicationBase.h"
+#include "RenderResource.h"
+
+#include <functional>
 
 #include "AnselFunctionLibrary.h"
 #include <AnselSDK.h>
@@ -35,7 +38,14 @@ static TAutoConsoleVariable<int32> CVarAllowHighQuality(
 static TAutoConsoleVariable<int32> CVarExtreme(
 	TEXT("r.Photography.Extreme"),
 	0,
-	TEXT("Whether to allow 'extreme' quality for Ansel RT.\n"),
+	TEXT("Whether to allow 'extreme' quality for Ansel RT (EXPERIMENTAL).\n"),
+	ECVF_RenderThreadSafe);
+
+// intentionally undocumented - debug flag
+static TAutoConsoleVariable<int32> CVarDebug0(
+	TEXT("r.Photography.Debug0"),
+	0,
+	TEXT("Debug - kill RT when in high-quality(!) mode\n"),
 	ECVF_RenderThreadSafe);
 
 /////////////////////////////////////////////////
@@ -65,7 +75,6 @@ public:
 		control_bloomintensity,
 		control_bloomscale,
 		control_scenefringeintensity,
-		control_motionbluramount,
 		control_COUNT
 	};
 	typedef union {
@@ -91,12 +100,14 @@ private:
 	bool BlueprintModifyCamera(ansel::Camera& InOutAnselCam, APlayerCameraManager* PCMgr); // returns whether modified cam is in original (session-start) position
 
 	void ConfigureRenderingSettingsForPhotography(FPostProcessSettings& InOutPostProcessSettings);
+	void SetUpSessionCVars();
 	void DoCustomUIControls(FPostProcessSettings& InOutPPSettings, bool bRebuildControls);
 	void DeclareSlider(int id, FText LocTextLabel, float LowerBound, float UpperBound, float Val);
 	bool ProcessUISlider(int id, float& InOutVal);
 
 	bool CaptureCVar(FString CVarName);
-	void SetCapturedCVar(const char* CVarName, float valueIfNotReset, bool wantReset);
+	void SetCapturedCVarPredicated(const char* CVarName, float valueIfNotReset, std::function<bool(const float, const float)> comparison, bool wantReset, bool useExistingPriority);
+	void SetCapturedCVar(const char* CVarName, float valueIfNotReset, bool wantReset = false, bool useExistingPriority = false);
 
 	ansel::Configuration* AnselConfig;
 	ansel::Camera AnselCamera;
@@ -125,15 +136,20 @@ private:
 	bool bWereSubtitlesEnabledBeforeSession;
 	bool bWasFadingEnabledBeforeSession;
 	bool bWasScreenMessagesEnabledBeforeSession = false;
+	float fTimeDilationBeforeSession;
 
 	bool bCameraIsInOriginalState = true;
 
 	bool bAutoPostprocess;
 	bool bAutoPause;
+	bool bRayTracingEnabled = false;
+	bool bPausedInternally = false;
 
 	bool bHighQualityModeDesired = false;
 	bool bHighQualityModeIsSetup = false;
 
+	uint32_t NumFramesSinceSessionStart;
+
 	// members relating to the 'Game Settings' controls in the Ansel overlay UI
 	TStaticBitArray<256> bEffectUIAllowed;
 
@@ -182,7 +198,6 @@ FNVAnselCameraPhotographyPrivate::FNVAnselCameraPhotographyPrivate()
 	, bAnselCaptureNewlyFinished(false)
 	, bForceDisallow(false)
 	, bIsOrthoProjection(false)
-	, bUIControlsNeedRebuild(false)
 {
 	for (int i = 0; i < bEffectUIAllowed.Num(); ++i)
 	{
@@ -437,20 +452,6 @@ void FNVAnselCameraPhotographyPrivate::DoCustomUIControls(FPostProcessSettings&
 			);
 		}
 
-		if (bEffectUIAllowed[MotionBlur] &&
-			InOutPPSettings.MotionBlurAmount > 0.f)
-		{
-			// Kludge: Character-based motion blur now works with a free camera and during multi-part captures but when the session starts the effect seems to be overly intense.  This dampens it.
-			float MotionBlurAmount = FMath::Min(InOutPPSettings.MotionBlurAmount, 0.1f);
-
-			DeclareSlider(
-				control_motionbluramount,
-				LOCTEXT("control_motionbluramount", "Motion Blur"),
-				0.f, 1.f,
-				MotionBlurAmount
-			);
-		}
-
 		bUIControlsNeedRebuild = false;
 	}
 
@@ -497,10 +498,6 @@ void FNVAnselCameraPhotographyPrivate::DoCustomUIControls(FPostProcessSettings&
 	{
 		InOutPPSettings.bOverride_SceneFringeIntensity = 1;
 	}
-	if (ProcessUISlider(control_motionbluramount, InOutPPSettings.MotionBlurAmount))
-	{
-		InOutPPSettings.bOverride_MotionBlurAmount = 1;
-	}
 }
 
 bool FNVAnselCameraPhotographyPrivate::UpdateCamera(FMinimalViewInfo& InOutPOV, APlayerCameraManager* PCMgr)
@@ -528,17 +525,13 @@ bool FNVAnselCameraPhotographyPrivate::UpdateCamera(FMinimalViewInfo& InOutPOV,
 		APlayerController* PCOwner = PCMgr->GetOwningPlayerController();
 		check(PCOwner != nullptr);
 
+		++NumFramesSinceSessionStart;
+
 		if (bAnselCaptureNewlyActive)
 		{
 			PCMgr->OnPhotographyMultiPartCaptureStart();
 			bGameCameraCutThisFrame = true;
 			bAnselCaptureNewlyActive = false;
-
-			// check for Panini projection & disable it?
-
-			// force sync texture loading and/or boost LODs?
-			// -> r.Streaming.FullyLoadUsedTextures for 4.13
-			// -> r.Streaming.?? for 4.12
 		}
 
 		if (bAnselCaptureNewlyFinished)
@@ -575,7 +568,9 @@ bool FNVAnselCameraPhotographyPrivate::UpdateCamera(FMinimalViewInfo& InOutPOV,
 
 			if (bAutoPause && !bWasPausedBeforeSession)
 			{
+				PCOwner->GetWorldSettings()->SetTimeDilation(fTimeDilationBeforeSession);
 				PCOwner->SetPause(false);
+				bPausedInternally = false;
 			}
 
 			PCMgr->GetWorld()->bIsCameraMoveableWhenPaused = bWasMovableCameraBeforeSession;
@@ -596,7 +591,8 @@ bool FNVAnselCameraPhotographyPrivate::UpdateCamera(FMinimalViewInfo& InOutPOV,
 				if (foo.Value.cvar)
 					foo.Value.cvar->SetWithCurrentPriority(foo.Value.fInitialVal);
 			}
-			InitialCVarMap.Empty();
+			InitialCVarMap.Empty(); // clear saved cvar values
+
 			bHighQualityModeIsSetup = false;
 			PCMgr->OnPhotographySessionEnd(); // after unpausing
 
@@ -608,6 +604,8 @@ bool FNVAnselCameraPhotographyPrivate::UpdateCamera(FMinimalViewInfo& InOutPOV,
 
 			if (bAnselSessionNewlyActive)
 			{
+				NumFramesSinceSessionStart = 0;
+
 				PCMgr->OnPhotographySessionStart(); // before pausing
 
 				// copy these values to avoid mixup if the CVars are changed during capture callbacks
@@ -616,16 +614,21 @@ bool FNVAnselCameraPhotographyPrivate::UpdateCamera(FMinimalViewInfo& InOutPOV,
 
 				bAutoPause = !!CVarAutoPause->GetInt();
 				bAutoPostprocess = !!CVarAutoPostProcess->GetInt();
-
+				bRayTracingEnabled = IsRayTracingEnabled();
+				
 				// attempt to pause game
 				bWasPausedBeforeSession = PCOwner->IsPaused();
 				bWasMovableCameraBeforeSession = PCMgr->GetWorld()->bIsCameraMoveableWhenPaused;
 				PCMgr->GetWorld()->bIsCameraMoveableWhenPaused = true;
 				if (bAutoPause && !bWasPausedBeforeSession)
 				{
-					PCOwner->SetPause(true); // should we bother to set delegate to enforce pausedness until session end?  probably over-engineering.
+					fTimeDilationBeforeSession = PCOwner->GetWorldSettings()->TimeDilation;
+					PCOwner->GetWorldSettings()->SetTimeDilation(0.f); // kill character motion-blur, this looks better than setting the motion-blur level to 0 (which flickers) - kinda heavy-handed but the only way I've found to kill motion-blur while also preventing flicker
+					// we pause in a *future* frame so Slomo can kick-in properly
 				}
 
+				SetUpSessionCVars();
+
 				bWasScreenMessagesEnabledBeforeSession = GAreScreenMessagesEnabled;
 				GAreScreenMessagesEnabled = false;
 
@@ -670,6 +673,17 @@ bool FNVAnselCameraPhotographyPrivate::UpdateCamera(FMinimalViewInfo& InOutPOV,
 				}
 			}
 
+			// ensure 2 frames have passed before pausing so that 0-timedilation can kick-in and kill the motion-blur!
+			// why 2 frames rather than 1 (or even 0)?  dunno!  probably 1 frame for the new time dilation to go into effect and 1 more frame for the motion vectors to update.
+			if (NumFramesSinceSessionStart == 2)
+			{
+				if (bAutoPause && !bWasPausedBeforeSession)
+				{
+					PCOwner->SetPause(true);
+					bPausedInternally = true;
+				}
+			}
+
 			AnselCameraToFMinimalView(InOutPOV, AnselCamera);
 
 			AnselCameraPrevious = AnselCamera;
@@ -685,83 +699,96 @@ bool FNVAnselCameraPhotographyPrivate::UpdateCamera(FMinimalViewInfo& InOutPOV,
 	return bGameCameraCutThisFrame;
 }
 
-void FNVAnselCameraPhotographyPrivate::SetCapturedCVar(const char* CVarName, float valueIfNotReset, bool wantReset)
+void FNVAnselCameraPhotographyPrivate::SetCapturedCVarPredicated(const char* CVarName, float valueIfNotReset, std::function<bool(const float, const float)> comparison, bool wantReset, bool useExistingPriority)
 {
 	CVarInfo* info = nullptr;
 	if (InitialCVarMap.Contains(CVarName) || CaptureCVar(CVarName))
 	{
 		info = &InitialCVarMap[CVarName];
-		if (info->cvar)
-			info->cvar->SetWithCurrentPriority(wantReset ? info->fInitialVal : valueIfNotReset);
+		if (info->cvar && comparison(valueIfNotReset, info->fInitialVal))
+		{
+			if (useExistingPriority)
+				info->cvar->SetWithCurrentPriority(wantReset ? info->fInitialVal : valueIfNotReset);
+			else
+				info->cvar->Set(wantReset ? info->fInitialVal : valueIfNotReset);
+		}
 	}
 	if (!(info && info->cvar)) UE_LOG(LogAnsel, Log, TEXT("CVar used by Ansel not found: %s"), CVarName);
 }
 
+void FNVAnselCameraPhotographyPrivate::SetCapturedCVar(const char* CVarName, float valueIfNotReset, bool wantReset, bool useExistingPriority)
+{
+	SetCapturedCVarPredicated(CVarName, valueIfNotReset,
+		[](float, float) { return true; },
+		wantReset, useExistingPriority);
+}
+
 void FNVAnselCameraPhotographyPrivate::ConfigureRenderingSettingsForPhotography(FPostProcessSettings& InOutPostProcessingSettings)
 {
-	if (CVarAllowHighQuality.GetValueOnAnyThread() && bHighQualityModeDesired)
+#define QUALITY_CVAR(NAME,BOOSTVAL) SetCapturedCVar(NAME, BOOSTVAL, !bHighQualityModeDesired, true)
+#define QUALITY_CVAR_AT_LEAST(NAME,BOOSTVAL) SetCapturedCVarPredicated(NAME, BOOSTVAL, std::greater<float>(), !bHighQualityModeDesired, true)
+#define QUALITY_CVAR_AT_MOST(NAME,BOOSTVAL) SetCapturedCVarPredicated(NAME, BOOSTVAL, std::less<float>(), !bHighQualityModeDesired, true)
+#define QUALITY_CVAR_LOWPRIORITY_AT_LEAST(NAME,BOOSTVAL) SetCapturedCVarPredicated(NAME, BOOSTVAL, std::greater<float>(), !bHighQualityModeDesired, false)
+
+	if (CVarDebug0->GetInt()
+		&& CVarAllowHighQuality.GetValueOnAnyThread()
+		&& bHighQualityModeIsSetup != bHighQualityModeDesired)
 	{
-		// bring rendering up to (at least) 100% resolution
-		if (InOutPostProcessingSettings.ScreenPercentage < 100.f)
-		{
-			// note: won't override r.screenpercentage set from console, that takes precedence
-			InOutPostProcessingSettings.bOverride_ScreenPercentage = 1;
-			InOutPostProcessingSettings.ScreenPercentage = 100.f;
-		}
+		// Debug - makes HQ mode actually try to kill RT features
+		QUALITY_CVAR("r.RayTracing.GlobalIllumination", 0);
+		QUALITY_CVAR("r.RayTracing.Reflections", 0);
+		QUALITY_CVAR("r.RayTracing.Shadows", 0);
+		QUALITY_CVAR("r.RayTracing.Translucency", 0);
+		QUALITY_CVAR("r.RayTracing.AmbientOcclusion", 0);
+		UE_LOG(LogAnsel, Log, TEXT("Photography Debug0 mode actualized (enabled=%d)"), (int)bHighQualityModeDesired);
+		bHighQualityModeIsSetup = bHighQualityModeDesired;
 	}
-
-	// Pump up (or reset) the quality.  Details subject to change as the engine evolves.
-	if (CVarAllowHighQuality.GetValueOnAnyThread() && bHighQualityModeIsSetup != bHighQualityModeDesired)
+	else
+	// Pump up (or reset) the quality. 
+	if (CVarAllowHighQuality.GetValueOnAnyThread()
+		&& bHighQualityModeIsSetup != bHighQualityModeDesired
+		&& (bPausedInternally || !bAutoPause) // <- don't start overriding vars until truly paused
+		&& (!CVarDebug0->GetInt()))
 	{
-#define QUALITY_CVAR(NAME,BOOSTVAL) SetCapturedCVar(NAME, BOOSTVAL, !bHighQualityModeDesired)
+		// bring rendering up to (at least) 100% resolution, but won't override manually set value on console
+		QUALITY_CVAR_LOWPRIORITY_AT_LEAST("r.ScreenPercentage", 100);
 
 		// most of these similar to typical cinematic sg.* scalability settings, toned down a little for performance
 
 		// can be a mild help with reflections
 		QUALITY_CVAR("r.gbufferformat", 5); // 5 = highest precision
 
+		// bias various geometry LODs
+		QUALITY_CVAR_AT_MOST("r.staticmeshloddistancescale", 0.25f); // large quality bias
+		QUALITY_CVAR_AT_MOST("r.landscapelodbias", -2);
+		QUALITY_CVAR_AT_MOST("r.skeletalmeshlodbias", -2);
+
 		// ~sg.AntiAliasingQuality @ cine
 		QUALITY_CVAR("r.postprocessaaquality", 6); // 6 == max
 		QUALITY_CVAR("r.defaultfeature.antialiasing", 2); // TAA
 
 		// ~sg.EffectsQuality @ cinematic
-		QUALITY_CVAR("r.TranslucencyLightingVolumeDim", 64);
+		QUALITY_CVAR_AT_LEAST("r.TranslucencyLightingVolumeDim", 64);
 		QUALITY_CVAR("r.RefractionQuality", 2);
 		QUALITY_CVAR("r.SSR.Quality", 4);
 		// QUALITY_CVAR("r.SceneColorFormat", 4); // don't really want to mess with this
 		QUALITY_CVAR("r.TranslucencyVolumeBlur", 1);
-		QUALITY_CVAR("r.MaterialQualityLevel", 1);
+		QUALITY_CVAR("r.MaterialQualityLevel", 1); // 1==high, 2==medium!
 		QUALITY_CVAR("r.SSS.Scale", 1);
 		QUALITY_CVAR("r.SSS.SampleSet", 2);
 		QUALITY_CVAR("r.SSS.Quality", 1);
 		QUALITY_CVAR("r.SSS.HalfRes", 0);
-		QUALITY_CVAR("r.EmitterSpawnRateScale", 1.f); // not sure this has a point when game is paused though
+		QUALITY_CVAR_AT_LEAST("r.EmitterSpawnRateScale", 1.f); // not sure this has a point when game is paused though
 		QUALITY_CVAR("r.ParticleLightQuality", 2);
-
-		// kludge: detailmode=2 is nice for high-quality, but it resets motion blur so if we have motion blur then don't apply the new detail mode until the camera has been moved for the first time even if HQ mode is desired
-		if ((InOutPostProcessingSettings.MotionBlurAmount == 0.f) ||
-			(!bCameraIsInOriginalState) ||
-			(!bHighQualityModeDesired))
-		{
-			QUALITY_CVAR("r.DetailMode", 2);
-			if (CVarAllowHighQuality.GetValueOnAnyThread() && bHighQualityModeDesired) // hide motion blur UI now that it won't do anything more this session
-			{
-				if (UIControls[control_motionbluramount].info.userControlId > 0) // we are using id 0 as 'unused'
-				{
-					ansel::removeUserControl(UIControls[control_motionbluramount].info.userControlId);
-					UIControls[control_motionbluramount].info.userControlId = 0;
-				}
-			}
-		}
+		QUALITY_CVAR("r.DetailMode", 2);
 
 		// ~sg.PostProcessQuality @ cinematic
-		//QUALITY_CVAR("r.MotionBlurQuality", 4); // nope - don't want to risk resetting currently visible motion blur
 		QUALITY_CVAR("r.AmbientOcclusionMipLevelFactor", 0.4f);
 		QUALITY_CVAR("r.AmbientOcclusionMaxQuality", 100);
 		QUALITY_CVAR("r.AmbientOcclusionLevels", -1);
 		QUALITY_CVAR("r.AmbientOcclusionRadiusScale", 1.f);
 		QUALITY_CVAR("r.DepthOfFieldQuality", 4);
-		QUALITY_CVAR("r.RenderTargetPoolMin", 500); // ?
+		QUALITY_CVAR_AT_LEAST("r.RenderTargetPoolMin", 500); // ?
 		QUALITY_CVAR("r.LensFlareQuality", 3);
 		QUALITY_CVAR("r.SceneColorFringeQuality", 1);
 		QUALITY_CVAR("r.BloomQuality", 5);
@@ -774,7 +801,7 @@ void FNVAnselCameraPhotographyPrivate::ConfigureRenderingSettingsForPhotography(
 		QUALITY_CVAR("r.DOF.Gather.AccumulatorQuality", 1);
 		QUALITY_CVAR("r.DOF.Gather.PostfilterMethod", 1);
 		QUALITY_CVAR("r.DOF.Gather.EnableBokehSettings", 1);
-		QUALITY_CVAR("r.DOF.Gather.RingCount", 5);
+		QUALITY_CVAR_AT_LEAST("r.DOF.Gather.RingCount", 5);
 		QUALITY_CVAR("r.DOF.Scatter.ForegroundCompositing", 1);
 		QUALITY_CVAR("r.DOF.Scatter.BackgroundCompositing", 2);
 		QUALITY_CVAR("r.DOF.Scatter.EnableBokehSettings", 1);
@@ -787,25 +814,24 @@ void FNVAnselCameraPhotographyPrivate::ConfigureRenderingSettingsForPhotography(
 
 		// ~sg.TextureQuality @ cinematic
 		QUALITY_CVAR("r.Streaming.MipBias", 0);
-		QUALITY_CVAR("r.MaxAnisotropy", 16);
+		QUALITY_CVAR_AT_LEAST("r.MaxAnisotropy", 16);
 		QUALITY_CVAR("r.Streaming.MaxEffectiveScreenSize", 0);
 		// intentionally don't mess with streaming pool size, see 'CVarExtreme' section below
 
 		// ~sg.FoliageQuality @ cinematic
-		QUALITY_CVAR("foliage.DensityScale", 1.f);
-		QUALITY_CVAR("grass.DensityScale", 1.f);
+		QUALITY_CVAR_AT_LEAST("foliage.DensityScale", 1.f);
+		QUALITY_CVAR_AT_LEAST("grass.DensityScale", 1.f);
 
 		// ~sg.ViewDistanceQuality @ cine but only mild draw distance boost
-		QUALITY_CVAR("r.viewdistancescale", 2.0f); // or even more...?
-		QUALITY_CVAR("r.skeletalmeshlodbias", -2); // somewhat tested
+		QUALITY_CVAR_AT_LEAST("r.viewdistancescale", 2.0f); // or even more...?
 
 		// ~sg.ShadowQuality @ cinematic
-		QUALITY_CVAR("r.LightFunctionQuality", 1);
+		QUALITY_CVAR_AT_LEAST("r.LightFunctionQuality", 2);
 		QUALITY_CVAR("r.ShadowQuality", 5);
-		QUALITY_CVAR("r.Shadow.CSM.MaxCascades", 10);
-		QUALITY_CVAR("r.Shadow.MaxResolution", 4096);
-		QUALITY_CVAR("r.Shadow.MaxCSMResolution", 4096);
-		QUALITY_CVAR("r.Shadow.RadiusThreshold", 0.f);
+		QUALITY_CVAR_AT_LEAST("r.Shadow.CSM.MaxCascades", 10);
+		QUALITY_CVAR_AT_LEAST("r.Shadow.MaxResolution", 4096);
+		QUALITY_CVAR_AT_LEAST("r.Shadow.MaxCSMResolution", 4096);
+		QUALITY_CVAR_AT_MOST("r.Shadow.RadiusThreshold", 0.f);
 		QUALITY_CVAR("r.Shadow.DistanceScale", 1.f);
 		QUALITY_CVAR("r.Shadow.CSM.TransitionScale", 1.f);
 		QUALITY_CVAR("r.Shadow.PreShadowResolutionFactor", 1.f);
@@ -813,58 +839,116 @@ void FNVAnselCameraPhotographyPrivate::ConfigureRenderingSettingsForPhotography(
 		QUALITY_CVAR("r.VolumetricFog", 1);
 		QUALITY_CVAR("r.VolumetricFog.GridPixelSize", 4);
 		QUALITY_CVAR("r.VolumetricFog.GridSizeZ", 128);
-		QUALITY_CVAR("r.VolumetricFog.HistoryMissSupersampleCount", 16);
-		QUALITY_CVAR("r.LightMaxDrawDistanceScale", 2.f);
+		QUALITY_CVAR_AT_LEAST("r.VolumetricFog.HistoryMissSupersampleCount", 16);
+		QUALITY_CVAR_AT_LEAST("r.LightMaxDrawDistanceScale", 2.f);
 		QUALITY_CVAR("r.CapsuleShadows", 1);
 
+		// pump up the quality of raytracing features, though we won't necessarily turn them on if the game doesn't already have them enabled
+		if (bRayTracingEnabled)
+		{
+			/*** HIGH-QUALITY MODE DOES *NOT* FORCE GI ON ***/
+			QUALITY_CVAR_AT_MOST("r.RayTracing.GlobalIllumination.DiffuseThreshold", 0); // artifact avoidance
+			//QUALITY_CVAR_AT_LEAST("r.RayTracing.GlobalIllumination.MaxBounces", 1); // 1~=IQ cost:benefit sweet-spot
+
+			/*** HIGH-QUALITY MODE DOES *NOT* FORCE RT AO ON ***/
+			QUALITY_CVAR_AT_LEAST("r.RayTracing.AmbientOcclusion.SamplesPerPixel", 1); // haven't seen benefit from larger values
+
+			/*** HIGH-QUALITY MODE FORCES RT REFLECTIONS ON ***/
+			QUALITY_CVAR_AT_LEAST("r.RayTracing.Reflections.MaxBounces", 2); // sweet-spot
+			QUALITY_CVAR_AT_LEAST("r.RayTracing.Reflections.MaxRoughness", 0.9f); // speed hit
+			QUALITY_CVAR("r.RayTracing.Reflections.SortMaterials", 1); // usually some kind of perf win, especially w/above reflection quality
+			QUALITY_CVAR("r.RayTracing.Reflections.DirectLighting", 1);
+			//QUALITY_CVAR("r.RayTracing.Reflections.EmissiveAndIndirectLighting", 1);// curiously problematic, leave alone
+			QUALITY_CVAR_AT_LEAST("r.RayTracing.Reflections.Shadows", 1); // -1==auto, 0==off, 1==hard, 2==soft/area(requires high spp)
+			QUALITY_CVAR("r.RayTracing.Reflections.HeightFog", 1);
+			//QUALITY_CVAR_AT_LEAST("r.RayTracing.Reflections.SamplesPerPixel", 2); // -1==use pp vol // NOPE, don't touch spp right now: 1 is ok, ~10 is good, anywhere in-between is noisy
+			QUALITY_CVAR_AT_LEAST("r.RayTracing.Reflections.ScreenPercentage", 100);
+			QUALITY_CVAR("r.RayTracing.Reflections", 1); // FORCE ON: ignore postproc volume flag
+
+			/*** HIGH-QUALITY MODE DOES *NOT* FORCE RT TRANSLUCENCY ON ***/
+			QUALITY_CVAR_AT_LEAST("r.RayTracing.Translucency.MaxRoughness", 0.9f);
+			//QUALITY_CVAR_AT_LEAST("r.RayTracing.Translucency.MaxRefractionRays", 11); // buggy with grass, leave alone for now
+			QUALITY_CVAR_AT_LEAST("r.RayTracing.Translucency.Shadows", 1); // turn on at least
+			//QUALITY_CVAR("r.RayTracing.Translucency", -1); // 1==enabled always, ignore postproc volume flags -- NOPE, DON'T FORCE-ENABLE TRANSLUCENCY, IT MAKES EVERY SINGLE TRANSLUCENCY REFRACT or just plain disappear, too weird for random content (i.e. Infiltrator). -1 == explicitly marked-up volumes use RT
+
+			/*** HIGH-QUALITY MODE FORCES RT SHADOWS ON ***/
+			//QUALITY_CVAR_AT_LEAST("r.RayTracing.Shadow.SamplesPerPixel", 1); // 5==reduces stippling artifacts // >1 seems to do nothing extra now?
+			////QUALITY_CVAR("r.Shadow.Denoiser", 2); // "GScreenSpaceDenoiser witch may be overriden by a third party plugin"
+			QUALITY_CVAR_AT_LEAST("r.RayTracing.Shadows", 1); // 1==enableRT (default)
+		}
+
 		 // these are some extreme settings whose quality:risk ratio may be debatable or unproven
 		if (CVarExtreme->GetInt())
 		{
 			// great idea but not until I've proven that this isn't deadly or extremely slow on lower-spec machines:
 
-			QUALITY_CVAR("r.Streaming.LimitPoolSizeToVRAM", 0);
-			QUALITY_CVAR("r.Streaming.PoolSize", 3000); // cine - perhaps redundant when r.streaming.fullyloadusedtextures
+			QUALITY_CVAR("r.Streaming.LimitPoolSizeToVRAM", 0); // 0 is aggressive but is it safe? seems safe.
+			QUALITY_CVAR_AT_LEAST("r.Streaming.PoolSize", 3000); // cine - perhaps redundant when r.streaming.fullyloadusedtextures
 
-			QUALITY_CVAR("r.streaming.hlodstrategy", 0); // probably use 0 if using r.streaming.fullyloadusedtextures, else 2
-			QUALITY_CVAR("r.streaming.fullyloadusedtextures", 1); // pretty but what happens when overcommitted?  fatal?
-			QUALITY_CVAR("r.viewdistancescale", 10.f); // cinematic - extreme
+			QUALITY_CVAR("r.streaming.hlodstrategy", 2); // probably use 0 if using r.streaming.fullyloadusedtextures, else 2
+			//QUALITY_CVAR("r.streaming.fullyloadusedtextures", 1); // no - LODs oscillate when overcommitted
+			QUALITY_CVAR_AT_LEAST("r.viewdistancescale", 10.f); // cinematic - extreme
 
-			// just not hugely tested:
+			if (bRayTracingEnabled)
+			{
+				// higher-IQ thresholds
+				QUALITY_CVAR_AT_LEAST("r.RayTracing.Translucency.MaxRoughness", 1.f); // speed hit
+				QUALITY_CVAR_AT_LEAST("r.RayTracing.Reflections.MaxRoughness", 1.f); // speed hit
+
+				//QUALITY_CVAR("r.ambientocclusionstaticfraction", 0.f); // trust RT AO/GI...? - needs more testing, doesn't seem a big win
+
+				/*** EXTREME-QUALITY MODE FORCES GI ON ***/
+				// first, some IQ:speed tweaks to make GI speed practical
+				//
+				QUALITY_CVAR("r.RayTracing.GlobalIllumination.ScreenPercentage", 50); // 50% = this is actually a quality DROP by default but it makes the GI speed practical -- requires >>>=2spp though
+				QUALITY_CVAR_AT_MOST("r.RayTracing.GlobalIllumination.MaxRayDistance", 7500); // ditto; most of the IQ benefit, but often faster than default huge ray distance
+				QUALITY_CVAR_AT_LEAST("r.RayTracing.GlobalIllumination.SamplesPerPixel", 4); // at LEAST 2spp needed to reduce significant noise in some scenes, even up to 8+ helps
+				QUALITY_CVAR_AT_LEAST("r.RayTracing.GlobalIllumination.NextEventEstimationSamples", 16); // 2==default; 16 necessary for low-light conditions when using only 4spp, else get blotches.  raising estimation samples cheaper than raising spp.
+				QUALITY_CVAR_AT_LEAST("r.GlobalIllumination.Denoiser.ReconstructionSamples", 56/*=max*/); // better if only using 4spp @ quarter rez.  default is 16.
+				//QUALITY_CVAR_AT_LEAST("r.RayTracing.GlobalIllumination.MaxBounces", 3); // 2+ is sometimes slightly noticable, sloww
+				////QUALITY_CVAR("r.RayTracing.GlobalIllumination.EvalSkyLight", 1); // EXPERIMENTAL
+				QUALITY_CVAR("r.RayTracing.GlobalIllumination", 1); // FORCE ON: should be fast enough to not TDR(!) with screenpercentage=50... usually a fair IQ win with random content... hidden behind 'EXTREME' mode until I've exercised it more.
+
+				// just not hugely tested:
+				QUALITY_CVAR_AT_LEAST("r.RayTracing.StochasticRectLight.SamplesPerPixel", 4);
+				//QUALITY_CVAR("r.RayTracing.StochasticRectLight", 1); // 1==suspicious, probably broken
+				QUALITY_CVAR_AT_LEAST("r.RayTracing.SkyLight.SamplesPerPixel", 4); // default==-1 UNPROVEN TRY ME
+			}
 
+			// just not hugely tested:
 			QUALITY_CVAR("r.particlelodbias", -2);
 
 			// unproven or possibly buggy
 			//QUALITY_CVAR("r.streaming.useallmips", 1); // removes relative prioritization spec'd by app... unproven that this is a good idea
-			//QUALITY_CVAR("r.streaming.limitpoolsizetovram", 0); // 0 is aggressive but is it safe?
-			//QUALITY_CVAR("r.streaming.boost", 9999); // 0 = supposedly use all available vram, but it looks like 0 = buggy
+			//QUALITY_CVAR_AT_LEAST("r.streaming.boost", 9999); // 0 = supposedly use all available vram, but it looks like 0 = buggy
 		}
 
 #undef QUALITY_CVAR
+#undef QUALITY_CVAR_AT_LEAST
+#undef QUALITY_CVAR_AT_MOST
+#undef QUALITY_CVAR_LOWPRIORITY_AT_LEAST
 
 		UE_LOG(LogAnsel, Log, TEXT("Photography HQ mode actualized (enabled=%d)"), (int)bHighQualityModeDesired);
 		bHighQualityModeIsSetup = bHighQualityModeDesired;
 	}
 
-	// Always want these regardless of desired quality mode
-
-	SetCapturedCVar("r.oneframethreadlag", 1, false); // ansel needs frame latency to be predictable
-	SetCapturedCVar("r.streaming.numstaticcomponentsprocessedperframe", 0, false); // 0 = load all pending static geom now
-
-	// these are okay tweaks to streaming heuristics to reduce latency of full texture loads or minimize VRAM waste
-	SetCapturedCVar("r.disablelodfade", 1, false);
-	SetCapturedCVar("r.Streaming.MaxNumTexturesToStreamPerFrame", 0, false); // no limit
-	SetCapturedCVar("r.streaming.minmipforsplitrequest", 1, false); // strictly prioritize what's visible right now
-	SetCapturedCVar("r.streaming.hiddenprimitivescale", 0.001f, false); // hint to engine to deprioritize obscured textures...?
-	SetCapturedCVar("r.streaming.framesforfullupdate", 1, false); // recalc required LODs ASAP
-	SetCapturedCVar("r.Streaming.Boost", 1, false);
-
 	if (bAnselCaptureActive)
 	{
+		// camera doesn't linger in one place very long so maximize streaming rate
+		SetCapturedCVar("r.disablelodfade", 1);
+		SetCapturedCVar("r.streaming.framesforfullupdate", 1); // recalc required LODs ASAP
+		SetCapturedCVar("r.Streaming.MaxNumTexturesToStreamPerFrame", 0); // no limit
+		SetCapturedCVar("r.streaming.numstaticcomponentsprocessedperframe", 0); // 0 = load all pending static geom now
+
 		if (bAutoPostprocess)
 		{
 			// force-disable the standard postprocessing effects which are known to
 			// be problematic in multi-part shots
 
+			// nerf remaining motion blur
+			InOutPostProcessingSettings.bOverride_MotionBlurAmount = 1;
+			InOutPostProcessingSettings.MotionBlurAmount = 0.f;
+
 			// these effects tile poorly
 			InOutPostProcessingSettings.bOverride_BloomDirtMaskIntensity = 1;
 			InOutPostProcessingSettings.BloomDirtMaskIntensity = 0.f;
@@ -917,6 +1001,21 @@ void FNVAnselCameraPhotographyPrivate::ConfigureRenderingSettingsForPhotography(
 	}
 }
 
+void FNVAnselCameraPhotographyPrivate::SetUpSessionCVars()
+{
+	// This set of CVar tweaks are good - or necessary - for photographic sessions
+	{
+		SetCapturedCVar("r.oneframethreadlag", 1); // ansel needs frame latency to be predictable
+
+		// these are okay tweaks to streaming heuristics to reduce latency of full texture loads or minimize VRAM waste
+		SetCapturedCVar("r.streaming.minmipforsplitrequest", 1); // strictly prioritize what's visible right now
+		SetCapturedCVar("r.streaming.hiddenprimitivescale", 0.001f); // hint to engine to deprioritize obscured textures...?
+		SetCapturedCVar("r.Streaming.Boost", 1);
+
+		SetCapturedCVar("r.motionblurquality", 0); // this nerfs motion blur for non-characters
+	}
+}
+
 void FNVAnselCameraPhotographyPrivate::UpdatePostProcessing(FPostProcessSettings& InOutPostProcessingSettings)
 {
 	if (bAnselSessionActive)
diff --git a/Engine/Plugins/Runtime/Nvidia/Ansel/Source/Ansel/Private/AnselFunctionLibrary.cpp b/Engine/Plugins/Runtime/Nvidia/Ansel/Source/Ansel/Private/AnselFunctionLibrary.cpp
--- a/Engine/Plugins/Runtime/Nvidia/Ansel/Source/Ansel/Private/AnselFunctionLibrary.cpp
+++ b/Engine/Plugins/Runtime/Nvidia/Ansel/Source/Ansel/Private/AnselFunctionLibrary.cpp
@@ -135,6 +135,13 @@ void UAnselFunctionLibrary::SetUIControlVisibility(UObject* WorldContextObject,
 
 void UAnselFunctionLibrary::ConstrainCameraByDistance(UObject* WorldContextObject, const FVector NewCameraLocation, const FVector PreviousCameraLocation, const FVector OriginalCameraLocation, FVector& OutCameraLocation, float MaxDistance)
 {
+	if (MaxDistance < 0.f)
+	{
+		// no constraint by distance
+		OutCameraLocation = NewCameraLocation;
+		return;
+	}
+
 	FVector MovementVector = NewCameraLocation - OriginalCameraLocation;
 	MovementVector = MovementVector.GetClampedToMaxSize(MaxDistance);
 	OutCameraLocation = OriginalCameraLocation + MovementVector;
@@ -146,6 +153,13 @@ void UAnselFunctionLibrary::ConstrainCameraByGeometry(UObject* WorldContextObjec
 
 	static IConsoleVariable* CVarConstrainCameraSize = IConsoleManager::Get().FindConsoleVariable(TEXT("r.Photography.Constrain.CameraSize"));
 	const float CameraRadius = CVarConstrainCameraSize->GetFloat();
+
+	if (CameraRadius < 0.f)
+	{
+		// no constraint by collisions
+		return;
+	}
+
 	const float OpenSpaceRadius = 2.f * CameraRadius; // Minimum free space around camera for it to be considered unconfined
 
 	static FVector LastUnconfinedScreenshotCamera = OriginalCameraLocation;
diff --git a/Engine/Source/Runtime/Engine/Private/Camera/CameraPhotography.cpp b/Engine/Source/Runtime/Engine/Private/Camera/CameraPhotography.cpp
--- a/Engine/Source/Runtime/Engine/Private/Camera/CameraPhotography.cpp
+++ b/Engine/Source/Runtime/Engine/Private/Camera/CameraPhotography.cpp
@@ -47,12 +47,12 @@ static TAutoConsoleVariable<float> CVarPhotographyTranslationSpeed(
 static TAutoConsoleVariable<float> CVarConstrainCameraSize(
 	TEXT("r.Photography.Constrain.CameraSize"),
 	14.0f,
-	TEXT("Radius (in Unreal Units) of sphere around the camera; used to prevent the camera clipping into nearby geometry when constraining camera with collision (Default: 14.0)"));
+	TEXT("Radius (in Unreal Units) of sphere around the camera; used to prevent the camera clipping into nearby geometry when constraining camera with collision.  Negative values disable default camera collisions. (Default: 14.0)"));
 
 static TAutoConsoleVariable<float> CVarConstrainCameraDistance(
 	TEXT("r.Photography.Constrain.MaxCameraDistance"),
 	2500.0f,
-	TEXT("Maximum distance (in Unreal Units) which camera is allowed to wander from its initial position when constraining camera by distance (Default: 2500.0)"));
+	TEXT("Maximum distance (in Unreal Units) which camera is allowed to wander from its initial position when constraining camera by distance.  Negative values disable default distance contraints. (Default: 2500.0)"));
 
 static TAutoConsoleVariable<int32> CVarPhotographyAutoPostprocess(
 	TEXT("r.Photography.AutoPostprocess"),
