Fixed compilation warning.

diff --git a/Engine/Plugins/Runtime/Nvidia/Ansel/Source/Ansel/Private/Ansel.cpp b/Engine/Plugins/Runtime/Nvidia/Ansel/Source/Ansel/Private/Ansel.cpp
--- a/Engine/Plugins/Runtime/Nvidia/Ansel/Source/Ansel/Private/Ansel.cpp
+++ b/Engine/Plugins/Runtime/Nvidia/Ansel/Source/Ansel/Private/Ansel.cpp
@@ -743,13 +743,13 @@ void FNVAnselCameraPhotographyPrivate::ConfigureRenderingSettingsForPhotography(
 		UE_LOG(LogAnsel, Log, TEXT("Photography Debug0 mode actualized (enabled=%d)"), (int)bHighQualityModeDesired);
 		bHighQualityModeIsSetup = bHighQualityModeDesired;
 	}
-	else
-	// Pump up (or reset) the quality. 
-	if (CVarAllowHighQuality.GetValueOnAnyThread()
+	else if (CVarAllowHighQuality.GetValueOnAnyThread()
 		&& bHighQualityModeIsSetup != bHighQualityModeDesired
 		&& (bPausedInternally || !bAutoPause) // <- don't start overriding vars until truly paused
 		&& (!CVarDebug0->GetInt()))
 	{
+		// Pump up (or reset) the quality. 
+
 		// bring rendering up to (at least) 100% resolution, but won't override manually set value on console
 		QUALITY_CVAR_LOWPRIORITY_AT_LEAST("r.ScreenPercentage", 100);
 
